Fix a reference to DataNodes becoming invalidated from a second add to the structure

diff --git a/Engine/Source/Editor/StructViewer/Private/SStructViewer.cpp b/Engine/Source/Editor/StructViewer/Private/SStructViewer.cpp
--- a/Engine/Source/Editor/StructViewer/Private/SStructViewer.cpp
+++ b/Engine/Source/Editor/StructViewer/Private/SStructViewer.cpp
@@ -907,22 +907,26 @@ void FStructHierarchy::PopulateStructHierarchy()
 			{
 				const UScriptStruct* SuperStruct = Cast<UScriptStruct>(CurrentStruct->GetSuperStruct());
 
-				TSharedPtr<FStructViewerNodeData>& ParentEntry = DataNodes.FindOrAdd(SuperStruct);
-				if (!ParentEntry.IsValid())
+				TSharedPtr<FStructViewerNodeData>& ParentEntryRef = DataNodes.FindOrAdd(SuperStruct);
+
+				// Need to have a pointer in-case the ref moves to avoid an extra look up in the map
+				TSharedPtr<FStructViewerNodeData> ParentEntry = ParentEntryRef;
+				if (!ParentEntryRef.IsValid())
 				{
 					check(SuperStruct); // The null entry should have been created above
-					ParentEntry = MakeShared<FStructViewerNodeData>(SuperStruct);
+					ParentEntryRef = MakeShared<FStructViewerNodeData>(SuperStruct);
 				}
 
-				TSharedPtr<FStructViewerNodeData>& MyEntry = DataNodes.FindOrAdd(CurrentStruct);
-				if (!MyEntry.IsValid())
+				TSharedPtr<FStructViewerNodeData>& MyEntryRef = DataNodes.FindOrAdd(CurrentStruct);
+				if (!MyEntryRef.IsValid())
 				{
-					MyEntry = MakeShared<FStructViewerNodeData>(CurrentStruct);
+					MyEntryRef = MakeShared<FStructViewerNodeData>(CurrentStruct);
 				}
 
+				// Need to re-acquire the reference in the struct as it may have been invalidated from a re-size
 				if (!Visited.Contains(CurrentStruct))
 				{
-					ParentEntry->AddChild(MyEntry.ToSharedRef());
+					ParentEntry->AddChild(MyEntryRef.ToSharedRef());
 					Visited.Add(CurrentStruct);
 				}
 
