[NO NON-PII COMMIT MESSAGE]

diff --git a/Engine/Source/Editor/StructViewer/Private/SStructViewer.cpp b/Engine/Source/Editor/StructViewer/Private/SStructViewer.cpp
--- a/Engine/Source/Editor/StructViewer/Private/SStructViewer.cpp
+++ b/Engine/Source/Editor/StructViewer/Private/SStructViewer.cpp
@@ -866,13 +866,15 @@ void FStructHierarchy::PopulateStructHierarchy()
 		DataNodes.Add(nullptr, StructRootNode);
 
 		TSet<const UScriptStruct*> Visited;
+		UPackage* TransientPackage = GetTransientPackage();
 
 		// Go through all of the structs and see if they should be added to the list.
 		for (TObjectIterator<UScriptStruct> StructIt; StructIt; ++StructIt)
 		{
 			const UScriptStruct* CurrentStruct = *StructIt;
-			if (Visited.Contains(CurrentStruct))
+			if (Visited.Contains(CurrentStruct) || CurrentStruct->GetOutermost() == TransientPackage)
 			{
+				// Skip transient structs as they are dead leftovers from user struct editing
 				continue;
 			}
 
@@ -921,7 +923,11 @@ void FStructHierarchy::PopulateStructHierarchy()
 
 		for (const FAssetData& UserDefinedStructData : UserDefinedStructsList)
 		{
-			StructRootNode->AddChild(MakeShared<FStructViewerNodeData>(UserDefinedStructData));
+			if (!UserDefinedStructData.IsAssetLoaded())
+			{
+				// If the asset is loaded it was added by the object iterator
+				StructRootNode->AddChild(MakeShared<FStructViewerNodeData>(UserDefinedStructData));
+			}
 		}
 	}
 
