[NO NON-PII COMMIT MESSAGE]

diff --git a/Engine/Source/Editor/StructViewer/Private/SStructViewer.cpp b/Engine/Source/Editor/StructViewer/Private/SStructViewer.cpp
--- a/Engine/Source/Editor/StructViewer/Private/SStructViewer.cpp
+++ b/Engine/Source/Editor/StructViewer/Private/SStructViewer.cpp
@@ -983,14 +983,20 @@ void FStructHierarchy::AddAsset(const FAssetData& InAddedAssetData)
 	FAssetRegistryModule& AssetRegistryModule = FModuleManager::LoadModuleChecked<FAssetRegistryModule>(TEXT("AssetRegistry"));
 	if (!AssetRegistryModule.Get().IsLoadingAssets())
 	{
-		// Make sure that the node does not already exist. There is a bit of double adding going on at times and this prevents it.
-		if (!FindNodeByStructPath(InAddedAssetData.ObjectPath))
+		// Only handle structs
+		UClass* AssetClass = InAddedAssetData.GetClass();
+
+		if (AssetClass && AssetClass->IsChildOf(UScriptStruct::StaticClass()))
 		{
-			// User defined structs are always root level structs
-			StructRootNode->AddChild(MakeShared<FStructViewerNodeData>(InAddedAssetData));
+			// Make sure that the node does not already exist. There is a bit of double adding going on at times and this prevents it.
+			if (!FindNodeByStructPath(InAddedAssetData.ObjectPath))
+			{
+				// User defined structs are always root level structs
+				StructRootNode->AddChild(MakeShared<FStructViewerNodeData>(InAddedAssetData));
 
-			// All Viewers must repopulate.
-			PopulateStructViewerDelegate.Broadcast();
+				// All Viewers must repopulate.
+				PopulateStructViewerDelegate.Broadcast();
+			}
 		}
 	}
 }
