Fixing up comment text.
#rnx

diff --git a/Engine/Source/Editor/StructViewer/Private/SStructViewer.cpp b/Engine/Source/Editor/StructViewer/Private/SStructViewer.cpp
--- a/Engine/Source/Editor/StructViewer/Private/SStructViewer.cpp
+++ b/Engine/Source/Editor/StructViewer/Private/SStructViewer.cpp
@@ -159,7 +159,7 @@ private:
 	/** Dirty the struct hierarchy so it will be rebuilt on the next call to UpdateStructHierarchy */
 	void DirtyStructHierarchy();
 
-	/** Populates the struct hierarchy tree, pulling all the loaded and unloaded structs into a master data tree */
+	/** Populates the struct hierarchy tree, pulling in all the loaded and unloaded structs. */
 	void PopulateStructHierarchy();
 
 	/**
