FNames containing asset paths are deprecated. FSoftObjectPath or FTopLevelAssetPath should be used instead.
Fixups for class features in the editor.

diff --git a/Engine/Plugins/Experimental/StructUtils/Source/StructUtilsEditor/Private/Customizations/InstancedStructDetails.cpp b/Engine/Plugins/Experimental/StructUtils/Source/StructUtilsEditor/Private/Customizations/InstancedStructDetails.cpp
--- a/Engine/Plugins/Experimental/StructUtils/Source/StructUtilsEditor/Private/Customizations/InstancedStructDetails.cpp
+++ b/Engine/Plugins/Experimental/StructUtils/Source/StructUtilsEditor/Private/Customizations/InstancedStructDetails.cpp
@@ -55,7 +55,7 @@ public:
 		return !BaseStruct || InStruct->IsChildOf(BaseStruct);
 	}
 
-	virtual bool IsUnloadedStructAllowed(const FStructViewerInitializationOptions& InInitOptions, const FName InStructPath, TSharedRef<FStructViewerFilterFuncs> InFilterFuncs) override
+	virtual bool IsUnloadedStructAllowed(const FStructViewerInitializationOptions& InInitOptions, const FSoftObjectPath& InStructPath, TSharedRef<FStructViewerFilterFuncs> InFilterFuncs) override
 	{
 		// User Defined Structs don't support inheritance, so only include them requested
 		return bAllowUserDefinedStructs;
diff --git a/Engine/Source/Editor/ClassViewer/Private/ClassViewerFilter.cpp b/Engine/Source/Editor/ClassViewer/Private/ClassViewerFilter.cpp
--- a/Engine/Source/Editor/ClassViewer/Private/ClassViewerFilter.cpp
+++ b/Engine/Source/Editor/ClassViewer/Private/ClassViewerFilter.cpp
@@ -669,12 +669,12 @@ PRAGMA_ENABLE_DEPRECATION_WARNINGS
 	bool bPassesAssetReferenceFilter = true;
 	if (AssetReferenceFilter.IsValid() && bIsBlueprintBase)
 	{
-		FName BlueprintAssetPath = FName(*GeneratedClassPathString);
-		FAssetData BlueprintAssetData = AssetRegistry.GetAssetByObjectPath(BlueprintAssetPath);
+		FString BlueprintAssetPath = GeneratedClassPathString;
+		FAssetData BlueprintAssetData = AssetRegistry.GetAssetByObjectPath(FSoftObjectPath(BlueprintAssetPath));
 		if (!BlueprintAssetData.IsValid())
 		{
-			BlueprintAssetPath = FName(*GeneratedClassPathString.LeftChop(2)); // Chop off _C
-			BlueprintAssetData = AssetRegistry.GetAssetByObjectPath(BlueprintAssetPath);
+			BlueprintAssetPath = GeneratedClassPathString.LeftChop(2); // Chop off _C
+			BlueprintAssetData = AssetRegistry.GetAssetByObjectPath(FSoftObjectPath(BlueprintAssetPath));
 		}
 
 		if (BlueprintAssetData.IsValid())
diff --git a/Engine/Source/Editor/ClassViewer/Private/ClassViewerNode.cpp b/Engine/Source/Editor/ClassViewer/Private/ClassViewerNode.cpp
--- a/Engine/Source/Editor/ClassViewer/Private/ClassViewerNode.cpp
+++ b/Engine/Source/Editor/ClassViewer/Private/ClassViewerNode.cpp
@@ -149,7 +149,7 @@ bool FClassViewerNode::IsClassPlaceable() const
 
 bool FClassViewerNode::IsBlueprintClass() const
 {
-	return BlueprintAssetPath != NAME_None;
+	return !BlueprintAssetPath.IsNull();
 }
 
 bool FClassViewerNode::IsEditorOnlyClass() const
diff --git a/Engine/Source/Editor/ClassViewer/Private/ClassViewerNode.h b/Engine/Source/Editor/ClassViewer/Private/ClassViewerNode.h
--- a/Engine/Source/Editor/ClassViewer/Private/ClassViewerNode.h
+++ b/Engine/Source/Editor/ClassViewer/Private/ClassViewerNode.h
@@ -99,7 +99,7 @@ public:
 	FTopLevelAssetPath ParentClassPath;
 
 	/** Full path to the Blueprint that this class is loaded from, none for native classes*/
-	FName BlueprintAssetPath;
+	FSoftObjectPath BlueprintAssetPath;
 
 	/** true if the class passed the filter. */
 	bool bPassesFilter;
diff --git a/Engine/Source/Editor/ClassViewer/Private/SClassViewer.cpp b/Engine/Source/Editor/ClassViewer/Private/SClassViewer.cpp
--- a/Engine/Source/Editor/ClassViewer/Private/SClassViewer.cpp
+++ b/Engine/Source/Editor/ClassViewer/Private/SClassViewer.cpp
@@ -1327,7 +1327,7 @@ void FClassHierarchy::CreateOrUpdateUnloadedClassNode(TSharedPtr<FClassViewerNod
 
 	// Blueprint-specific fields
 
-	InOutClassViewerNode->BlueprintAssetPath = InAssetData.ObjectPath;
+	InOutClassViewerNode->BlueprintAssetPath = InAssetData.GetSoftObjectPath();
 
 	// It is an unloaded blueprint, so we need to create the structure that will hold the data.
 	TSharedPtr<FUnloadedBlueprintData> UnloadedBlueprintData = MakeShareable(new FUnloadedBlueprintData(InOutClassViewerNode));
@@ -1376,7 +1376,7 @@ void FClassHierarchy::PopulateClassHierarchy()
 			else
 			{
 				UE_LOG(LogEditorClassViewer, Warning, TEXT("AssetRegistry Blueprint %s is missing tag value for %s. Blueprint will not be available to ClassViewer when unloaded."),
-					*AssetData.ObjectPath.ToString(), *FBlueprintTags::GeneratedClassPath.ToString());
+					*AssetData.GetObjectPathString(), *FBlueprintTags::GeneratedClassPath.ToString());
 			}
 		}
 
@@ -1385,7 +1385,7 @@ void FClassHierarchy::PopulateClassHierarchy()
 
 		for (const FAssetData& AssetData : Assets)
 		{
-			FTopLevelAssetPath ClassPathNameFromAssetPath(AssetData.ObjectPath.ToString());
+			FTopLevelAssetPath ClassPathNameFromAssetPath = AssetData.GetSoftObjectPath().GetAssetPath();
 			TSharedPtr<FClassViewerNode>& Node = ClassPathToNode.FindOrAdd(ClassPathNameFromAssetPath);
 			CreateOrUpdateUnloadedClassNode(Node, AssetData, ClassPathNameFromAssetPath);
 		}
@@ -2115,7 +2115,7 @@ FReply SClassViewer::OnDragDetected( const FGeometry& Geometry, const FPointerEv
 					return FReply::Handled().BeginDragDrop(FClassDragDropOp::New(MakeWeakObjectPtr(Class)));
 				}	
 			}
-			else if (Item->BlueprintAssetPath != NAME_None)
+			else if (!Item->BlueprintAssetPath.IsNull())
 			{
 				IAssetRegistry& AssetRegistry = FModuleManager::LoadModuleChecked<FAssetRegistryModule>(AssetRegistryConstants::ModuleName).Get();
 
diff --git a/Engine/Source/Editor/GameProjectGeneration/Private/SNewClassDialog.cpp b/Engine/Source/Editor/GameProjectGeneration/Private/SNewClassDialog.cpp
--- a/Engine/Source/Editor/GameProjectGeneration/Private/SNewClassDialog.cpp
+++ b/Engine/Source/Editor/GameProjectGeneration/Private/SNewClassDialog.cpp
@@ -1432,7 +1432,7 @@ void SNewClassDialog::UpdateInputValidity()
 		if (bLastInputValidityCheckSuccessful)
 		{
 			IAssetRegistry& AssetRegistry = FModuleManager::LoadModuleChecked<FAssetRegistryModule>(AssetRegistryConstants::ModuleName).Get();
-			if (AssetRegistry.GetAssetByObjectPath(*(NewClassPath / NewClassName)).IsValid())
+			if (AssetRegistry.GetAssetByObjectPath(FSoftObjectPath(NewClassPath / NewClassName)).IsValid())
 			{
 				bLastInputValidityCheckSuccessful = false;
 				LastInputValidityErrorText = FText::Format(LOCTEXT("AssetAlreadyExists", "An asset called {0} already exists in {1}."), FText::FromString(NewClassName), FText::FromString(NewClassPath));
diff --git a/Engine/Source/Editor/GraphEditor/Private/KismetPins/SGraphPinStruct.cpp b/Engine/Source/Editor/GraphEditor/Private/KismetPins/SGraphPinStruct.cpp
--- a/Engine/Source/Editor/GraphEditor/Private/KismetPins/SGraphPinStruct.cpp
+++ b/Engine/Source/Editor/GraphEditor/Private/KismetPins/SGraphPinStruct.cpp
@@ -78,7 +78,7 @@ public:
 		return !MetaStruct || InStruct->IsChildOf(MetaStruct);
 	}
 
-	virtual bool IsUnloadedStructAllowed(const FStructViewerInitializationOptions& InInitOptions, const FName InStructPath, TSharedRef<FStructViewerFilterFuncs> InFilterFuncs) override
+	virtual bool IsUnloadedStructAllowed(const FStructViewerInitializationOptions& InInitOptions, const FSoftObjectPath& InStructPath, TSharedRef<FStructViewerFilterFuncs> InFilterFuncs) override
 	{
 		// User Defined Structs don't support inheritance, so only include them if we have don't a MetaStruct set
 		return MetaStruct == nullptr;
diff --git a/Engine/Source/Editor/PropertyEditor/Private/UserInterface/PropertyEditor/SPropertyEditorStruct.cpp b/Engine/Source/Editor/PropertyEditor/Private/UserInterface/PropertyEditor/SPropertyEditorStruct.cpp
--- a/Engine/Source/Editor/PropertyEditor/Private/UserInterface/PropertyEditor/SPropertyEditorStruct.cpp
+++ b/Engine/Source/Editor/PropertyEditor/Private/UserInterface/PropertyEditor/SPropertyEditorStruct.cpp
@@ -32,7 +32,7 @@ public:
 		return !MetaStruct || InStruct->IsChildOf(MetaStruct);
 	}
 
-	virtual bool IsUnloadedStructAllowed(const FStructViewerInitializationOptions& InInitOptions, const FName InStructPath, TSharedRef<FStructViewerFilterFuncs> InFilterFuncs) override
+	virtual bool IsUnloadedStructAllowed(const FStructViewerInitializationOptions& InInitOptions, const FSoftObjectPath& InStructPath, TSharedRef<FStructViewerFilterFuncs> InFilterFuncs) override
 	{
 		// User Defined Structs don't support inheritance, so only include them if we have don't a MetaStruct set
 		return MetaStruct == nullptr;
diff --git a/Engine/Source/Editor/StructViewer/Private/SStructViewer.cpp b/Engine/Source/Editor/StructViewer/Private/SStructViewer.cpp
--- a/Engine/Source/Editor/StructViewer/Private/SStructViewer.cpp
+++ b/Engine/Source/Editor/StructViewer/Private/SStructViewer.cpp
@@ -146,8 +146,8 @@ public:
 	 * @param InStructPath	The path name of the struct to find the node for.
 	 * @return The node.
 	 */
-	TSharedPtr<FStructViewerNodeData> FindNodeByStructPath(const TSharedRef<FStructViewerNodeData>& InRootNode, const FName InStructPath);
-	TSharedPtr<FStructViewerNodeData> FindNodeByStructPath(const FName InStructPath)
+	TSharedPtr<FStructViewerNodeData> FindNodeByStructPath(const TSharedRef<FStructViewerNodeData>& InRootNode, const FSoftObjectPath& InStructPath);
+	TSharedPtr<FStructViewerNodeData> FindNodeByStructPath(const FSoftObjectPath& InStructPath)
 	{
 		return FindNodeByStructPath(GetStructRootNode(), InStructPath);
 	}
@@ -170,8 +170,8 @@ private:
 	 *
 	 * @return Returns true if the struct was found and deleted successfully.
 	 */
-	bool FindAndRemoveNodeByStructPath(const TSharedRef<FStructViewerNodeData>& InRootNode, const FName InStructPath);
-	bool FindAndRemoveNodeByStructPath(const FName InStructPath)
+	bool FindAndRemoveNodeByStructPath(const TSharedRef<FStructViewerNodeData>& InRootNode, const FSoftObjectPath& InStructPath);
+	bool FindAndRemoveNodeByStructPath(const FSoftObjectPath& InStructPath)
 	{
 		return FindAndRemoveNodeByStructPath(GetStructRootNode(), InStructPath);
 	}
@@ -227,7 +227,7 @@ namespace StructViewer
 		 * @param InInitOptions		The struct viewer's options, holds the AllowedStructs and DisallowedStructs.
 		 * @param InStructPath		The path name to test against.
 		 */
-		bool IsStructAllowed_UnloadedStruct(const FStructViewerInitializationOptions& InInitOptions, const FName InStructPath)
+		bool IsStructAllowed_UnloadedStruct(const FStructViewerInitializationOptions& InInitOptions, const FSoftObjectPath& InStructPath)
 		{
 			if (InInitOptions.StructFilter.IsValid())
 			{
@@ -640,9 +640,9 @@ public:
 
 				ToolTip = IDocumentation::Get()->CreateToolTip(RestrictionToolTip, nullptr, "", "");
 			}
-			else if (!AssociatedNode->GetStructPath().IsNone())
+			else if (!AssociatedNode->GetStructPath().IsNull())
 			{
-				ToolTip = SNew(SToolTip).Text(FText::FromName(AssociatedNode->GetStructPath()));
+				ToolTip = SNew(SToolTip).Text(FText::FromString(AssociatedNode->GetStructPath().ToString()));
 			}
 
 			return ToolTip;
@@ -938,7 +938,7 @@ void FStructHierarchy::PopulateStructHierarchy()
 	PopulateStructViewerDelegate.Broadcast();
 }
 
-TSharedPtr<FStructViewerNodeData> FStructHierarchy::FindNodeByStructPath(const TSharedRef<FStructViewerNodeData>& InRootNode, const FName InStructPath)
+TSharedPtr<FStructViewerNodeData> FStructHierarchy::FindNodeByStructPath(const TSharedRef<FStructViewerNodeData>& InRootNode, const FSoftObjectPath& InStructPath)
 {
 	// Check if the current node is the struct path that is being searched for
 	if (InRootNode->GetStructPath() == InStructPath)
@@ -960,7 +960,7 @@ TSharedPtr<FStructViewerNodeData> FStructHierarchy::FindNodeByStructPath(const T
 	return nullptr;
 }
 
-bool FStructHierarchy::FindAndRemoveNodeByStructPath(const TSharedRef<FStructViewerNodeData>& InRootNode, const FName InStructPath)
+bool FStructHierarchy::FindAndRemoveNodeByStructPath(const TSharedRef<FStructViewerNodeData>& InRootNode, const FSoftObjectPath& InStructPath)
 {
 	// Check if the current node contains a child of struct path that is being searched for
 	if (InRootNode->RemoveChild(InStructPath))
@@ -992,7 +992,7 @@ void FStructHierarchy::AddAsset(const FAssetData& InAddedAssetData)
 		if (AssetClass && AssetClass->IsChildOf(UScriptStruct::StaticClass()))
 		{
 			// Make sure that the node does not already exist. There is a bit of double adding going on at times and this prevents it.
-			if (!FindNodeByStructPath(InAddedAssetData.ObjectPath))
+			if (!FindNodeByStructPath(InAddedAssetData.GetSoftObjectPath()))
 			{
 				// User defined structs are always root level structs
 				StructRootNode->AddChild(MakeShared<FStructViewerNodeData>(InAddedAssetData));
@@ -1006,7 +1006,7 @@ void FStructHierarchy::AddAsset(const FAssetData& InAddedAssetData)
 
 void FStructHierarchy::RemoveAsset(const FAssetData& InRemovedAssetData)
 {
-	if (FindAndRemoveNodeByStructPath(InRemovedAssetData.ObjectPath))
+	if (FindAndRemoveNodeByStructPath(InRemovedAssetData.GetSoftObjectPath()))
 	{
 		// All viewers must refresh.
 		PopulateStructViewerDelegate.Broadcast();
@@ -1782,7 +1782,7 @@ void SStructViewer::Populate()
 		// Take the package names for the internal only structs and convert them into their UScriptStructs
 		for (const TSoftObjectPtr<const UScriptStruct>& InternalStructName : InternalStructNames)
 		{
-			const TSharedPtr<FStructViewerNodeData> StructNode = FStructHierarchy::Get().FindNodeByStructPath(*InternalStructName.ToString());
+			const TSharedPtr<FStructViewerNodeData> StructNode = FStructHierarchy::Get().FindNodeByStructPath(InternalStructName.ToSoftObjectPath());
 			if (StructNode.IsValid())
 			{
 				if (const UScriptStruct* Struct = StructNode->GetStruct())
diff --git a/Engine/Source/Editor/StructViewer/Private/SStructViewer.h b/Engine/Source/Editor/StructViewer/Private/SStructViewer.h
--- a/Engine/Source/Editor/StructViewer/Private/SStructViewer.h
+++ b/Engine/Source/Editor/StructViewer/Private/SStructViewer.h
@@ -216,7 +216,7 @@ private:
 	bool bSaveExpansionStates;
 
 	/** The map holding the expansion state map for the tree. */
-	TMap<FName, bool> ExpansionStateMap;
+	TMap<FSoftObjectPath, bool> ExpansionStateMap;
 
 	/** True if the Struct Viewer needs to be repopulated at the next appropriate opportunity, occurs whenever structs are added, removed, renamed, etc. */
 	bool bNeedsRefresh;
diff --git a/Engine/Source/Editor/StructViewer/Private/StructViewerNode.cpp b/Engine/Source/Editor/StructViewer/Private/StructViewerNode.cpp
--- a/Engine/Source/Editor/StructViewer/Private/StructViewerNode.cpp
+++ b/Engine/Source/Editor/StructViewer/Private/StructViewerNode.cpp
@@ -37,7 +37,7 @@ FStructViewerNodeData::FStructViewerNodeData(const UScriptStruct* InStruct)
 
 FStructViewerNodeData::FStructViewerNodeData(const FAssetData& InStructAsset)
 	: StructName(InStructAsset.AssetName.ToString())
-	, StructPath(InStructAsset.ObjectPath)
+	, StructPath(InStructAsset.GetSoftObjectPath())
 {
 	// Attempt to find the struct asset in the case where it's already been loaded
 	Struct = FindObject<UScriptStruct>(nullptr, *StructPath.ToString());
@@ -66,7 +66,7 @@ bool FStructViewerNodeData::LoadStruct() const
 	}
 
 	// Attempt to load the struct
-	if (!StructPath.IsNone())
+	if (!StructPath.IsNull())
 	{
 		FScopedSlowTask SlowTask(0.0f, LOCTEXT("LoadingStruct", "Loading Struct..."));
 		SlowTask.MakeDialogDelayed(1.0f);
@@ -103,7 +103,7 @@ void FStructViewerNodeData::AddUniqueChild(const TSharedRef<FStructViewerNodeDat
 	}
 }
 
-bool FStructViewerNodeData::RemoveChild(const FName InStructPath)
+bool FStructViewerNodeData::RemoveChild(const FSoftObjectPath& InStructPath)
 {
 	for (auto ChildIt = ChildNodes.CreateIterator(); ChildIt; ++ChildIt)
 	{
diff --git a/Engine/Source/Editor/StructViewer/Private/StructViewerNode.h b/Engine/Source/Editor/StructViewer/Private/StructViewerNode.h
--- a/Engine/Source/Editor/StructViewer/Private/StructViewerNode.h
+++ b/Engine/Source/Editor/StructViewer/Private/StructViewerNode.h
@@ -9,6 +9,7 @@
 #include "Templates/SharedPointer.h"
 #include "UObject/Class.h"
 #include "UObject/NameTypes.h"
+#include "UObject/SoftObjectPath.h"
 #include "UObject/WeakObjectPtrTemplates.h"
 
 class IPropertyHandle;
@@ -42,13 +43,13 @@ public:
 	}
 
 	/** Get the full object path to the struct we represent */
-	FName GetStructPath() const
+	FSoftObjectPath GetStructPath() const
 	{
 		return StructPath;
 	}
 
 	/** Get the full object path to the parent of the struct we represent */
-	FName GetParentStructPath() const
+	FSoftObjectPath GetParentStructPath() const
 	{
 		return ParentStructPath;
 	}
@@ -84,7 +85,7 @@ public:
 	void AddUniqueChild(const TSharedRef<FStructViewerNodeData>& InChild);
 
 	/** Remove the child representing the given struct path (if present) */
-	bool RemoveChild(const FName InStructPath);
+	bool RemoveChild(const FSoftObjectPath& InStructPath);
 
 private:
 	/** The unlocalized name of the struct we represent */
@@ -94,10 +95,10 @@ private:
 	mutable FText StructDisplayName;
 
 	/** The full object path to the struct we represent */
-	FName StructPath;
+	FSoftObjectPath StructPath;
 
 	/** The full object path to the parent of the struct we represent */
-	FName ParentStructPath;
+	FSoftObjectPath ParentStructPath;
 
 	/** The struct that we represent (for loaded struct assets, or native structs) */
 	mutable TWeakObjectPtr<const UScriptStruct> Struct;
@@ -135,13 +136,13 @@ public:
 	FText GetStructDisplayName(const EStructViewerNameTypeToDisplay InNameType) const;
 
 	/** Get the full object path to the struct we represent */
-	FName GetStructPath() const
+	FSoftObjectPath GetStructPath() const
 	{
 		return NodeData->GetStructPath();
 	}
 
 	/** Get the full object path to the parent of the struct we represent */
-	FName GetParentStructPath() const
+	FSoftObjectPath GetParentStructPath() const
 	{
 		return NodeData->GetParentStructPath();
 	}
diff --git a/Engine/Source/Editor/StructViewer/Public/StructViewerFilter.h b/Engine/Source/Editor/StructViewer/Public/StructViewerFilter.h
--- a/Engine/Source/Editor/StructViewer/Public/StructViewerFilter.h
+++ b/Engine/Source/Editor/StructViewer/Public/StructViewerFilter.h
@@ -27,7 +27,7 @@ public:
 	 * @param InStructPath				The path of the unloaded struct to be tested.
 	 * @param InFilterFuncs				Useful functions for filtering.
 	 */
-	virtual bool IsUnloadedStructAllowed(const FStructViewerInitializationOptions& InInitOptions, const FName InStructPath, TSharedRef<class FStructViewerFilterFuncs> InFilterFuncs) = 0;
+	virtual bool IsUnloadedStructAllowed(const FStructViewerInitializationOptions& InInitOptions, const FSoftObjectPath& InStructPath, TSharedRef<class FStructViewerFilterFuncs> InFilterFuncs) = 0;
 };
 
 enum class EStructFilterReturn : uint8
diff --git a/Engine/Source/Editor/UnrealEd/Private/ClassIconFinder.cpp b/Engine/Source/Editor/UnrealEd/Private/ClassIconFinder.cpp
--- a/Engine/Source/Editor/UnrealEd/Private/ClassIconFinder.cpp
+++ b/Engine/Source/Editor/UnrealEd/Private/ClassIconFinder.cpp
@@ -117,7 +117,7 @@ const UClass* FClassIconFinder::GetIconClassForAssetData(const FAssetData& InAss
 			*bOutIsClassType = true;
 		}
 
-		return FindObject<UClass>(nullptr, *InAssetData.ObjectPath.ToString());
+		return FindObject<UClass>(nullptr, *InAssetData.GetObjectPathString());
 	}
 	
 	static const FName IgnoreClassThumbnail(TEXT("IgnoreClassThumbnail"));
diff --git a/Engine/Source/Editor/UnrealEd/Private/EditorClassUtils.cpp b/Engine/Source/Editor/UnrealEd/Private/EditorClassUtils.cpp
--- a/Engine/Source/Editor/UnrealEd/Private/EditorClassUtils.cpp
+++ b/Engine/Source/Editor/UnrealEd/Private/EditorClassUtils.cpp
@@ -373,27 +373,9 @@ FTopLevelAssetPath FEditorClassUtils::GetClassPathNameFromAssetTag(const FAssetD
 
 FName FEditorClassUtils::GetClassPathFromAsset(const FAssetData& InAssetData, bool bGenerateClassPathIfMissing /*= false*/)
 {
-	bool bIsBPGC = false;
-	const bool bIsBP = IsBlueprintAsset(InAssetData, &bIsBPGC);
-
-	if (bIsBPGC)
-	{
-		return InAssetData.ObjectPath;
-	}
-	else if (bIsBP)
-	{
 PRAGMA_DISABLE_DEPRECATION_WARNINGS
-		FName ClassPath = GetClassPathFromAssetTag(InAssetData);
+	return GetClassPathNameFromAsset(InAssetData, bGenerateClassPathIfMissing).ToFName();
 PRAGMA_ENABLE_DEPRECATION_WARNINGS
-		if (bGenerateClassPathIfMissing && ClassPath.IsNone())
-		{
-			FNameBuilder ClassPathBuilder(InAssetData.ObjectPath);
-			ClassPathBuilder << "_C";
-			ClassPath = FName(ClassPathBuilder.ToString());
-		}
-		return ClassPath;
-	}
-	return NAME_None;
 }
 
 FTopLevelAssetPath FEditorClassUtils::GetClassPathNameFromAsset(const FAssetData& InAssetData, bool bGenerateClassPathIfMissing /*= false*/)
@@ -403,14 +385,14 @@ FTopLevelAssetPath FEditorClassUtils::GetClassPathNameFromAsset(const FAssetData
 
 	if (bIsBPGC)
 	{
-		return FTopLevelAssetPath(InAssetData.ObjectPath.ToString());
+		return FTopLevelAssetPath(InAssetData.GetSoftObjectPath().GetAssetPath());
 	}
 	else if (bIsBP)
 	{
 		FTopLevelAssetPath ClassPath = GetClassPathNameFromAssetTag(InAssetData);
 		if (bGenerateClassPathIfMissing && ClassPath.IsNull())
 		{
-			FString ClassPathString = InAssetData.ObjectPath.ToString();
+			FString ClassPathString = InAssetData.GetObjectPathString();
 			ClassPathString += TEXT("_C");
 			ClassPath = FTopLevelAssetPath(ClassPathString);
 		}
diff --git a/Engine/Source/Editor/UnrealEd/Private/Factories/DataTableFactory.cpp b/Engine/Source/Editor/UnrealEd/Private/Factories/DataTableFactory.cpp
--- a/Engine/Source/Editor/UnrealEd/Private/Factories/DataTableFactory.cpp
+++ b/Engine/Source/Editor/UnrealEd/Private/Factories/DataTableFactory.cpp
@@ -37,7 +37,7 @@ bool UDataTableFactory::ConfigureProperties()
 			return FDataTableEditorUtils::IsValidTableStruct(InStruct);
 		}
 
-		virtual bool IsUnloadedStructAllowed(const FStructViewerInitializationOptions& InInitOptions, const FName InStructPath, TSharedRef<FStructViewerFilterFuncs> InFilterFuncs) override
+		virtual bool IsUnloadedStructAllowed(const FStructViewerInitializationOptions& InInitOptions, const FSoftObjectPath& InStructPath, TSharedRef<FStructViewerFilterFuncs> InFilterFuncs) override
 		{
 			// Unloaded structs are always User Defined Structs, and User Defined Structs are always allowed
 			// They will be re-validated by IsStructAllowed once loaded during the pick
diff --git a/Engine/Source/Editor/UnrealEd/Private/Factories/MirrorDataTableFactory.cpp b/Engine/Source/Editor/UnrealEd/Private/Factories/MirrorDataTableFactory.cpp
--- a/Engine/Source/Editor/UnrealEd/Private/Factories/MirrorDataTableFactory.cpp
+++ b/Engine/Source/Editor/UnrealEd/Private/Factories/MirrorDataTableFactory.cpp
@@ -30,7 +30,7 @@ public:
 		return FDataTableEditorUtils::IsValidTableStruct(InStruct) && InStruct->FindPropertyByName(TEXT("MirroredName")) != nullptr;
 	}
 
-	virtual bool IsUnloadedStructAllowed(const FStructViewerInitializationOptions& InInitOptions, const FName InStructPath, TSharedRef<FStructViewerFilterFuncs> InFilterFuncs) override
+	virtual bool IsUnloadedStructAllowed(const FStructViewerInitializationOptions& InInitOptions, const FSoftObjectPath& InStructPath, TSharedRef<FStructViewerFilterFuncs> InFilterFuncs) override
 	{
 		// Unloaded structs are always User Defined Structs, and User Defined Structs are always allowed
 		// They will be re-validated by IsStructAllowed once loaded during the pick
