PR #9840: Allow to show tooltips for items in StructViewer (Contributed by [REDACTED_PERSON])

diff --git a/Engine/Source/Editor/StructViewer/Private/SStructViewer.cpp b/Engine/Source/Editor/StructViewer/Private/SStructViewer.cpp
--- a/Engine/Source/Editor/StructViewer/Private/SStructViewer.cpp
+++ b/Engine/Source/Editor/StructViewer/Private/SStructViewer.cpp
@@ -638,11 +638,24 @@ public:
 				FText RestrictionToolTip;
 				PropertyHandle->GenerateRestrictionToolTip(*AssociatedNode->GetStructName(), RestrictionToolTip);
 
-				ToolTip = IDocumentation::Get()->CreateToolTip(RestrictionToolTip, nullptr, "", "");
+				ToolTip = SNew(SToolTip).Text(RestrictionToolTip);
 			}
 			else if (!AssociatedNode->GetStructPath().IsNull())
 			{
-				ToolTip = SNew(SToolTip).Text(FText::FromString(AssociatedNode->GetStructPath().ToString()));
+				const UScriptStruct* Struct = AssociatedNode->GetStruct();
+				if (Struct != nullptr && Struct->GetBoolMetaData("ShowTooltip"))
+				{
+					const FText ToolTipText = FText::Format(LOCTEXT("ToolTipFormat", "{0}\n\n{1}"), 
+						AssociatedNode->GetStruct()->GetToolTipText(), 
+						FText::FromString(AssociatedNode->GetStructPath().ToString())
+					);
+					
+					ToolTip = SNew(SToolTip).Text(ToolTipText);
+				}
+				else
+				{
+					ToolTip = SNew(SToolTip).Text(FText::FromString(AssociatedNode->GetStructPath().ToString()));	
+				}
 			}
 
 			return ToolTip;
