Set scroll position to the selected struct when opening struct picker. So that you see options near selected struct rather then always starting at the top and scrolling to the area you want to scroll to.

diff --git a/Engine/Plugins/Experimental/StructUtils/Source/StructUtilsEditor/Private/Customizations/InstancedStructDetails.cpp b/Engine/Plugins/Experimental/StructUtils/Source/StructUtilsEditor/Private/Customizations/InstancedStructDetails.cpp
--- a/Engine/Plugins/Experimental/StructUtils/Source/StructUtilsEditor/Private/Customizations/InstancedStructDetails.cpp
+++ b/Engine/Plugins/Experimental/StructUtils/Source/StructUtilsEditor/Private/Customizations/InstancedStructDetails.cpp
@@ -491,13 +491,17 @@ TSharedRef<SWidget> FInstancedStructDetails::GenerateStructPicker()
 	StructFilter->bAllowUserDefinedStructs = BaseScriptStruct == nullptr; // Only allow user defined structs when BaseStruct is not set.
 	StructFilter->bAllowBaseStruct = !bExcludeBaseStruct;
 
+	const UScriptStruct* SelectedStruct = nullptr;
+	const FPropertyAccess::Result Result = UE::StructUtils::Private::GetCommonScriptStruct(StructProperty, SelectedStruct);
+	
 	FStructViewerInitializationOptions Options;
 	Options.bShowNoneOption = bAllowNone;
 	Options.StructFilter = StructFilter;
 	Options.NameTypeToDisplay = EStructViewerNameTypeToDisplay::DisplayName;
 	Options.DisplayMode = bShowTreeView ? EStructViewerDisplayMode::TreeView : EStructViewerDisplayMode::ListView;
 	Options.bAllowViewOptions = !bHideViewOptions;
-
+	Options.SelectedStruct = SelectedStruct;
+	
 	FOnStructPicked OnPicked(FOnStructPicked::CreateSP(this, &FInstancedStructDetails::OnStructPicked));
 
 	return SNew(SBox)
diff --git a/Engine/Source/Editor/StructViewer/Private/SStructViewer.cpp b/Engine/Source/Editor/StructViewer/Private/SStructViewer.cpp
--- a/Engine/Source/Editor/StructViewer/Private/SStructViewer.cpp
+++ b/Engine/Source/Editor/StructViewer/Private/SStructViewer.cpp
@@ -1862,6 +1862,18 @@ void SStructViewer::Populate()
 		// Sort the list alphabetically.
 		RootTreeItems.Sort(&FStructViewerNode::SortPredicate);
 
+		// Scroll to selected struct
+		for (const TSharedPtr<FStructViewerNode>& Node : RootTreeItems)
+		{
+			if (Node.IsValid())
+			{
+				if (InitOptions.SelectedStruct == Node->GetStruct())
+				{
+					StructList->RequestScrollIntoView(Node);
+				}
+			}
+		}
+
 		// Only display this option if the user wants it and in Picker Mode.
 		if (InitOptions.bShowNoneOption && InitOptions.Mode == EStructViewerMode::StructPicker)
 		{
diff --git a/Engine/Source/Editor/StructViewer/Public/StructViewerModule.h b/Engine/Source/Editor/StructViewer/Public/StructViewerModule.h
--- a/Engine/Source/Editor/StructViewer/Public/StructViewerModule.h
+++ b/Engine/Source/Editor/StructViewer/Public/StructViewerModule.h
@@ -16,6 +16,8 @@ class UScriptStruct;
 
 /** Delegate used with the Struct Viewer in 'struct picking' mode.  You'll bind a delegate when the struct viewer widget is created, which will be fired off when a struct is selected in the list */
 DECLARE_DELEGATE_OneParam(FOnStructPicked, const UScriptStruct*);
+/** Delegate used with the Struct Viewer in 'struct picking' mode.  You'll bind a delegate when the struct viewer widget is created, which will be fired off when a list of structs is opened for selection */
+DECLARE_DELEGATE_OneParam(FOnStructPickerOpened, const UScriptStruct*);
 
 enum class EStructViewerMode : uint8
 {
@@ -93,6 +95,9 @@ public:
 	/** true (the default) shows a background border behind the struct viewer widget. */
 	bool bShowBackgroundBorder;
 
+	/** Defined currently selected struct to scroll to to when struct picker is opened. */
+	const UScriptStruct* SelectedStruct;
+
 	/** Defines additional structs you want listed in the "Common Structs" section for the picker. */
 	TArray<const UScriptStruct*> ExtraPickerCommonStructs;
 
@@ -108,6 +113,7 @@ public:
 		, NameTypeToDisplay(EStructViewerNameTypeToDisplay::StructName)
 		, bAllowViewOptions(true)
 		, bShowBackgroundBorder(true)
+	    , SelectedStruct(nullptr)
 	{
 	}
 };
