Implementing support for StructView to filter results based on display name as well. Sometimes what should be valid in filter seach by display name isn't valid because the filter only checks struct name.

diff --git a/Engine/Source/Editor/StructViewer/Private/SStructViewer.cpp b/Engine/Source/Editor/StructViewer/Private/SStructViewer.cpp
--- a/Engine/Source/Editor/StructViewer/Private/SStructViewer.cpp
+++ b/Engine/Source/Editor/StructViewer/Private/SStructViewer.cpp
@@ -322,13 +322,15 @@ namespace StructViewer
 			bool bReturnPassesFilter = false;
 			if (OriginalRootNodeStruct)
 			{
-				bReturnPassesFilter = bPassesDeveloperFilter && bPassesInternalFilter && IsStructAllowed(InInitOptions, OriginalRootNodeStruct) && PassesFilter(InOriginalRootNode->GetStructName(), InTextFilter);
+				bReturnPassesFilter = bPassesDeveloperFilter && bPassesInternalFilter && IsStructAllowed(InInitOptions, OriginalRootNodeStruct) &&
+					(PassesFilter(InOriginalRootNode->GetStructName(), InTextFilter) || PassesFilter(InOriginalRootNode->GetStructDisplayName().ToString(), InTextFilter));
 			}
 			else
 			{
 				if (bInShowUnloadedStructs)
 				{
-					bReturnPassesFilter = bPassesDeveloperFilter && bPassesInternalFilter && IsStructAllowed_UnloadedStruct(InInitOptions, InOutRootNode->GetStructPath()) && PassesFilter(InOriginalRootNode->GetStructName(), InTextFilter);
+					bReturnPassesFilter = bPassesDeveloperFilter && bPassesInternalFilter && IsStructAllowed_UnloadedStruct(InInitOptions, InOutRootNode->GetStructPath()) &&
+						(PassesFilter(InOriginalRootNode->GetStructName(), InTextFilter) || PassesFilter(InOriginalRootNode->GetStructDisplayName().ToString(), InTextFilter));
 				}
 			}
 			InOutRootNode->PassedFilter(bReturnPassesFilter);
@@ -455,13 +457,15 @@ namespace StructViewer
 			bool bPassedFilter = false;
 			if (OriginalRootNodeStruct)
 			{
-				bPassedFilter = bPassesDeveloperFilter && bPassesInternalFilter && IsStructAllowed(InInitOptions, OriginalRootNodeStruct) && PassesFilter(InOriginalRootNode->GetStructName(), InTextFilter);
+				bPassedFilter = bPassesDeveloperFilter && bPassesInternalFilter && IsStructAllowed(InInitOptions, OriginalRootNodeStruct) &&
+					(PassesFilter(InOriginalRootNode->GetStructName(), InTextFilter) || PassesFilter(InOriginalRootNode->GetStructDisplayName().ToString(), InTextFilter));
 			}
 			else
 			{
 				if (bInShowUnloadedStructs)
 				{
-					bPassedFilter = bPassesDeveloperFilter && bPassesInternalFilter && IsStructAllowed_UnloadedStruct(InInitOptions, InOriginalRootNode->GetStructPath()) && PassesFilter(InOriginalRootNode->GetStructName(), InTextFilter);
+					bPassedFilter = bPassesDeveloperFilter && bPassesInternalFilter && IsStructAllowed_UnloadedStruct(InInitOptions, InOriginalRootNode->GetStructPath()) &&
+						(PassesFilter(InOriginalRootNode->GetStructName(), InTextFilter) || PassesFilter(InOriginalRootNode->GetStructDisplayName().ToString(), InTextFilter));
 				}
 			}
 
