Struct Picker: Support showing icons (via the meta tag ShowIcon=true)

diff --git a/Engine/Source/Editor/StructViewer/Private/SStructViewer.cpp b/Engine/Source/Editor/StructViewer/Private/SStructViewer.cpp
--- a/Engine/Source/Editor/StructViewer/Private/SStructViewer.cpp
+++ b/Engine/Source/Editor/StructViewer/Private/SStructViewer.cpp
@@ -13,6 +13,7 @@
 
 #include "Editor.h"
 #include "Styling/AppStyle.h"
+#include "Styling/SlateIconFinder.h"
 #include "SlateOptMacros.h"
 #include "EditorWidgetsModule.h"
 #include "Framework/Commands/UIAction.h"
@@ -677,6 +678,16 @@ public:
 					SNew(SExpanderArrow, SharedThis(this))
 				]
 
+				+SHorizontalBox::Slot()
+				.AutoWidth()
+				.Padding(0.0f, 0.0f, 6.0f, 0.0f)
+				.VAlign(VAlign_Center)
+				[
+					SNew(SImage)
+					.Image(this, &SStructItem::GetIconImage)
+					.Visibility(this, &SStructItem::GetIconVisibility)
+				]
+
 				+SHorizontalBox::Slot()
 				.FillWidth(1.0f)
 				.Padding(0.0f, 3.0f, 6.0f, 3.0f)
@@ -782,6 +793,24 @@ private:
 		return TextColor;
 	}
 
+	const FSlateBrush* GetIconImage() const
+	{
+		if (!AssociatedNode)
+		{
+			return nullptr;
+		}
+
+		return FSlateIconFinder::FindIconBrushForClass(AssociatedNode->GetStruct(), "ClassIcon.Object");
+	}
+
+	EVisibility GetIconVisibility() const
+	{
+		const UScriptStruct* Struct = AssociatedNode ? AssociatedNode->GetStruct() : nullptr;
+		return Struct && Struct->GetBoolMetaDataHierarchical("ShowIcon") && GetIconImage()
+			? EVisibility::Visible
+			: EVisibility::Collapsed;
+	}
+
 private:
 	/** The struct name for which this item is associated with. */
 	FText StructDisplayName;
