Instanced Structs: Add ability to sort structs in a custom way inside pickers via a SortFunction UPROPERTY meta tag

diff --git a/Engine/Source/Editor/StructUtilsEditor/Private/SInstancedStructPicker.cpp b/Engine/Source/Editor/StructUtilsEditor/Private/SInstancedStructPicker.cpp
--- a/Engine/Source/Editor/StructUtilsEditor/Private/SInstancedStructPicker.cpp
+++ b/Engine/Source/Editor/StructUtilsEditor/Private/SInstancedStructPicker.cpp
@@ -309,6 +309,7 @@ TSharedRef<SWidget> SInstancedStructPicker::GenerateStructPicker()
 	Options.DisplayMode = bShowTreeView ? EStructViewerDisplayMode::TreeView : EStructViewerDisplayMode::ListView;
 	Options.bAllowViewOptions = !bHideViewOptions;
 	Options.SelectedStruct = SelectedStruct;
+	Options.PropertyHandle = StructProperty;
 
 	const FOnStructPicked OnPicked(FOnStructPicked::CreateSP(this, &SInstancedStructPicker::StructPicked));
 
diff --git a/Engine/Source/Editor/StructViewer/Private/SStructViewer.cpp b/Engine/Source/Editor/StructViewer/Private/SStructViewer.cpp
--- a/Engine/Source/Editor/StructViewer/Private/SStructViewer.cpp
+++ b/Engine/Source/Editor/StructViewer/Private/SStructViewer.cpp
@@ -1887,7 +1887,7 @@ void SStructViewer::Populate()
 		StructViewer::Helpers::GetStructList(InitOptions, RootTreeItems, *TextFilterPtr, bShowUnloadedStructs, GetCurrentDeveloperViewType(), ShowingInternalStructs, InternalStructs, InternalPaths);
 
 		// Sort the list alphabetically.
-		RootTreeItems.Sort(&FStructViewerNode::SortPredicate);
+		FStructViewerNode::SortNodes(RootTreeItems, InitOptions.PropertyHandle);
 
 		// Scroll to selected struct
 		for (const TSharedPtr<FStructViewerNode>& Node : RootTreeItems)
diff --git a/Engine/Source/Editor/StructViewer/Private/StructViewerNode.cpp b/Engine/Source/Editor/StructViewer/Private/StructViewerNode.cpp
--- a/Engine/Source/Editor/StructViewer/Private/StructViewerNode.cpp
+++ b/Engine/Source/Editor/StructViewer/Private/StructViewerNode.cpp
@@ -191,7 +191,7 @@ void FStructViewerNode::SortChildrenRecursive()
 
 void FStructViewerNode::SortChildren()
 {
-	ChildNodes.Sort(&SortPredicate);
+	SortNodes(ChildNodes, PropertyHandle);
 }
 
 bool FStructViewerNode::SortPredicate(const TSharedPtr<FStructViewerNode>& InA, const TSharedPtr<FStructViewerNode>& InB)
@@ -205,6 +205,67 @@ bool FStructViewerNode::SortPredicate(const TSharedPtr<FStructViewerNode>& InA,
 	return AString < BString;
 }
 
+void FStructViewerNode::SortNodes(TArray<TSharedPtr<FStructViewerNode>>& Nodes, const TSharedPtr<IPropertyHandle>& PropertyHandle)
+{
+	if (PropertyHandle.IsValid())
+	{
+		static const FName SortFunctionMetaName = "SortFunction";
+
+		const FString& FunctionName = PropertyHandle->GetMetaData(SortFunctionMetaName);
+
+		if (!FunctionName.IsEmpty())
+		{
+			const UFunction* Function = nullptr;
+
+			// Fully-qualified function name
+			if (UE::String::FindFirstChar(FunctionName, TEXT('.')) != INDEX_NONE)
+			{
+				Function = FindObject<UFunction>(nullptr, *FunctionName, true);
+			}
+			else
+			{
+				TArray<UObject*> Containers;
+				PropertyHandle->GetOuterObjects(Containers);
+
+				for (const UObject* Container : Containers)
+				{
+					Function = Container ? Container->FindFunction(FName(FunctionName)) : nullptr;
+					if (Function)
+					{
+						break;
+					}
+				}
+			}
+
+			if (ensureMsgf(Function && Function->HasAnyFunctionFlags(EFunctionFlags::FUNC_Static), TEXT("Invalid SortFunction: %s"), *FunctionName))
+			{
+				DECLARE_DELEGATE_RetVal_TwoParams(bool, FStructViewerUserCompare, const UScriptStruct*, const UScriptStruct*);
+
+				UObject* Container = Function->GetOuterUClass()->GetDefaultObject();
+				FStructViewerUserCompare CompareFunc = FStructViewerUserCompare::CreateUFunction(Container, Function->GetFName());
+
+				Nodes.Sort([&CompareFunc](const TSharedPtr<FStructViewerNode>& InA, const TSharedPtr<FStructViewerNode>& InB)
+				{
+					check(InA.IsValid() && InB.IsValid());
+
+					const UScriptStruct* AStruct = InA->GetStruct();
+					const UScriptStruct* BStruct = InB->GetStruct();
+					if (AStruct && BStruct)
+					{
+						return CompareFunc.Execute(InA->GetStruct(), InB->GetStruct());
+					}
+
+					return SortPredicate(InA, InB);
+				});
+
+				return;
+			}
+		}
+	}
+
+	Nodes.Sort(&SortPredicate);
+}
+
 bool FStructViewerNode::IsRestricted() const
 {
 	return PropertyHandle && PropertyHandle->IsRestricted(GetStructName());
diff --git a/Engine/Source/Editor/StructViewer/Private/StructViewerNode.h b/Engine/Source/Editor/StructViewer/Private/StructViewerNode.h
--- a/Engine/Source/Editor/StructViewer/Private/StructViewerNode.h
+++ b/Engine/Source/Editor/StructViewer/Private/StructViewerNode.h
@@ -195,6 +195,9 @@ public:
 	/** Predicate function that can be used to sort instances by struct name */
 	static bool SortPredicate(const TSharedPtr<FStructViewerNode>& InA, const TSharedPtr<FStructViewerNode>& InB);
 
+	/** Sorts the array of nodes either by a custom sort function, or by using the default one (::SortPredicate) */
+	static void SortNodes(TArray<TSharedPtr<FStructViewerNode>>& Nodes, const TSharedPtr<IPropertyHandle>& PropertyHandle = { });
+
 	/** Check whether this struct is restricted for the specific context */
 	bool IsRestricted() const;
 
