Struct Picker: Do not show "None" item if it does not match the text search filter

diff --git a/Engine/Source/Editor/StructViewer/Private/SStructViewer.cpp b/Engine/Source/Editor/StructViewer/Private/SStructViewer.cpp
--- a/Engine/Source/Editor/StructViewer/Private/SStructViewer.cpp
+++ b/Engine/Source/Editor/StructViewer/Private/SStructViewer.cpp
@@ -1904,9 +1904,15 @@ void SStructViewer::Populate()
 		// Only display this option if the user wants it and in Picker Mode.
 		if (InitOptions.bShowNoneOption && InitOptions.Mode == EStructViewerMode::StructPicker)
 		{
+			TSharedRef<FStructViewerNode> NoneNode = MakeShared<FStructViewerNode>();
+
 			// @todo - It would seem smart to add this in before the other items, since it needs to be on top. However, that causes strange issues with saving/restoring expansion states. 
 			// This is likely not very efficient since the list can have hundreds and even thousands of items.
-			RootTreeItems.Insert(MakeShared<FStructViewerNode>(), 0);
+			if (StructViewer::Helpers::PassesFilter(NoneNode->GetStructName(), *TextFilterPtr) ||
+				StructViewer::Helpers::PassesFilter(NoneNode->GetStructDisplayName().ToString(), *TextFilterPtr))
+			{
+				RootTreeItems.Insert(MoveTemp(NoneNode), 0);
+			}
 		}
 
 		NumStructs = 0;
