PR #12846: Add "Copy Struct Path" menu item to Struct Viewer

diff --git a/Engine/Source/Editor/StructViewer/Private/SStructViewer.cpp b/Engine/Source/Editor/StructViewer/Private/SStructViewer.cpp
--- a/Engine/Source/Editor/StructViewer/Private/SStructViewer.cpp
+++ b/Engine/Source/Editor/StructViewer/Private/SStructViewer.cpp
@@ -46,6 +46,7 @@
 #include "PropertyHandle.h"
 #include "Logging/LogMacros.h"
 #include "SourceCodeNavigation.h"
+#include "HAL/PlatformApplicationMisc.h"
 #include "Subsystems/AssetEditorSubsystem.h"
 
 #define LOCTEXT_NAMESPACE "SStructViewer"
@@ -532,6 +533,17 @@ namespace StructViewer
 			FSourceCodeNavigation::NavigateToStruct(InStruct);
 		}
 
+		/**
+		 * Copy struct FTopLevelAssetPath to clipboard.
+		 */
+		void CopyStructPathToClipboard(const UScriptStruct* InStruct)
+		{
+			if (InStruct)
+			{
+				FPlatformApplicationMisc::ClipboardCopy(*InStruct->GetStructPathName().ToString());
+			}
+		}
+
 		/**
 		 * Finds the struct in the content browser.
 		 */
@@ -578,6 +590,13 @@ namespace StructViewer
 						FUIAction(FExecuteAction::CreateStatic(&StructViewer::Helpers::OpenStructInIDE, InStruct))
 						);
 				}
+				
+				MenuBuilder.AddMenuEntry(
+					LOCTEXT("StructViewerMenuCopyStructPath", "Copy Struct Path"),
+					LOCTEXT("StructViewerMenuCopyStructPath_Tooltip", "Copies the path for this struct (e.g '/Script/Module.StructName') to the clipboard."),
+					FSlateIcon(),
+					FUIAction(FExecuteAction::CreateStatic(&StructViewer::Helpers::CopyStructPathToClipboard, InStruct))
+					);
 			}
 
 			return MenuBuilder.MakeWidget();
diff --git a/Engine/Source/Editor/StructViewer/StructViewer.Build.cs b/Engine/Source/Editor/StructViewer/StructViewer.Build.cs
--- a/Engine/Source/Editor/StructViewer/StructViewer.Build.cs
+++ b/Engine/Source/Editor/StructViewer/StructViewer.Build.cs
@@ -25,7 +25,8 @@ public class StructViewer : ModuleRules
 				"EditorFramework",
 				"UnrealEd",
 				"ContentBrowserData",
-                "Settings",
+				"Settings",
+				"ApplicationCore",
             }
 		);
 
